Add python script for generating an html color palette for the FColorList class.

diff --git a/Engine/Extras/ColorListPalette/generate_colorlist_palette.py b/Engine/Extras/ColorListPalette/generate_colorlist_palette.py
new file mode 100644
--- /dev/null
+++ b/Engine/Extras/ColorListPalette/generate_colorlist_palette.py
@@ -0,0 +1,134 @@
+import argparse
+import os.path
+import re
+
+colorlist_cpp_rel_path = os.path.join("Source", "Runtime", "Core", "Private", "Math", "ColorList.cpp")
+
+re_color_def = re.compile(
+    r'^const FColor FColorList::(?P<name>\w+)\s*\(\s*(?P<r>\d+),\s*(?P<g>\d+),\s*(?P<b>\d+),\s*(?P<a>\d+)\s*\);',
+    re.MULTILINE)
+
+def find_colorlist():
+    cur_dir = os.path.dirname(__file__)
+    engine_dir = os.path.normpath(os.path.join(cur_dir, "..", ".."))
+    colorlist_cpp = os.path.join(engine_dir, colorlist_cpp_rel_path)
+    return colorlist_cpp
+
+def parse_colorlist(colorlist_cpp):
+    print("Reading %s" % colorlist_cpp)
+    with open(colorlist_cpp, "r") as fp:
+        cpp_code = fp.read()
+
+    colors = {}
+    matches = re_color_def.finditer(cpp_code)
+    for match in matches:
+        colors[match.group('name')] = (match.group('r'), match.group('g'), match.group('b'), match.group('a'))
+    print("Found %d colors" % len(colors))
+    return colors
+
+def generate_palette_html(colors, fp):
+    fp.write(
+"""<!DOCTYPE html>
+<html lang="en">
+<head>
+  <meta charset="utf-8">
+  <meta name="viewport" content="width=device-width, initial-scale=1">
+  <title>Unreal Engine Common Colors</title>
+  <style>
+    body {
+      font-family: Arial, sans-serif;
+      background: #f9f9f9;
+      padding: 2rem;
+    }
+    h1 {
+        text-align: center;
+    }
+    .color-grid {
+      display: grid;
+      gap: 1rem;
+      margin: 0;
+    }
+    .color-item {
+      display: grid;
+      grid-template-columns: 1fr 4fr;
+      align-items: center;
+      gap: 1rem;
+      background: #fff;
+      padding: 0.5rem;
+      border: 1px solid #ddd;
+      border-radius: 0.3rem;
+      box-shadow: 0 2px 4px rgba(0, 0, 0, 0.05);
+    }
+    .color-box {
+      width: 24px;
+      height: 24px;
+      border: 1px solid #aaa;
+    }
+    .color-name {
+      font-size: 14px;
+      color: #333;
+    }
+    .footer {
+      text-align: center;
+      font-size: 0.8em;
+      padding: 1rem;
+      color: #999;
+    }
+    .footer code {
+      font-size: 1.1em;
+      color: #73a;
+    }
+    @media (min-width: 600px) {
+        .color-grid { grid-template-columns: repeat(2, 1fr); }
+    }
+    @media (min-width: 900px) {
+        .color-grid { grid-template-columns: repeat(3, 1fr); }
+    }
+    @media (min-width: 1200px) {
+        .color-grid { grid-template-columns: repeat(4, 1fr); }
+    }
+    @media (min-width: 1500px) {
+        .color-grid { width: calc(1500px - 4rem); margin: 0 auto; }
+    }
+  </style>
+</head>
+<body>
+<h1>Unreal Engine Common Colors</h1>
+<div class="color-grid">
+""")
+
+    for name, color in colors.items():
+        fp.write(
+f"""  <div class="color-item">
+        <div class="color-box" style="background-color: rgb({color[0]},{color[1]},{color[2]});"></div>
+        <div class="color-name">{name}</div>
+    </div>\n""")
+
+    fp.write(
+f"""</div>
+<div class="footer">
+    <p>generated from <code>{colorlist_cpp_rel_path}</code></p>
+</div>
+</body>
+</html>""")
+
+def main():
+    parser = argparse.ArgumentParser()
+    parser.add_argument('-i', '--input')
+    parser.add_argument('out_path', default="colorlist.html", nargs='?')
+
+    args = parser.parse_args()
+
+    colorlist_cpp = args.input
+    if not colorlist_cpp:
+        colorlist_cpp = find_colorlist()
+
+    colors = parse_colorlist(colorlist_cpp)
+
+    print("Writing %s" % args.out_path)
+    with open(args.out_path, "w") as fp:
+        generate_palette_html(colors, fp)
+
+if __name__ == '__main__':
+    main()
+
